- Niagara MovieRenderQueue plugin
- Contains a data interface for inspecting MRQ information

diff --git a/Engine/Plugins/FX/NiagaraMRQ/NiagaraMRQ.uplugin b/Engine/Plugins/FX/NiagaraMRQ/NiagaraMRQ.uplugin
new file mode 100644
--- /dev/null
+++ b/Engine/Plugins/FX/NiagaraMRQ/NiagaraMRQ.uplugin
@@ -0,0 +1,35 @@
+{
+	"FileVersion": 3,
+	"Version": 1,
+	"VersionName": "1.0",
+	"FriendlyName": "Niagara MRQ Support",
+	"Description": "Contains a data interface that can be used to read Movie Render Queue information in Niagara simulations.",
+	"Category": "FX",
+	"CreatedBy": "Epic Games, Inc.",
+	"CreatedByURL": "https://epicgames.com",
+	"DocsURL": "",
+	"MarketplaceURL": "",
+	"SupportURL": "",
+	"EnabledByDefault": false,
+	"CanContainContent": true,
+	"IsBetaVersion": true,
+	"IsExperimentalVersion": true,
+	"Installed": false,
+	"Modules": [
+		{
+			"Name": "NiagaraMRQ",
+			"Type": "Runtime",
+			"LoadingPhase": "PreDefault"
+		},
+	],
+	"Plugins": [
+		{
+			"Name": "MovieRenderPipeline",
+			"Enabled": true
+		},
+		{
+			"Name": "Niagara",
+			"Enabled": true
+		}
+	]
+}
\ No newline at end of file
diff --git a/Engine/Plugins/FX/NiagaraMRQ/Shaders/Private/NiagaraDataInterfaceMRQ.ush b/Engine/Plugins/FX/NiagaraMRQ/Shaders/Private/NiagaraDataInterfaceMRQ.ush
new file mode 100644
--- /dev/null
+++ b/Engine/Plugins/FX/NiagaraMRQ/Shaders/Private/NiagaraDataInterfaceMRQ.ush
@@ -0,0 +1,14 @@
+// Copyright Epic Games, Inc. All Rights Reserved.
+
+int		{ParameterName}_Active;
+int		{ParameterName}_TemporalSampleCount;
+int		{ParameterName}_TemporalSampleIndex;
+float	{ParameterName}_SequenceFPS;
+
+void GetMRQInfo_{ParameterName}(out bool bActive, out int TemporalSampleCount, out int TemporalSampleIndex, out float SequenceFPS)
+{
+	bActive					= {ParameterName}_Active != 0;
+	TemporalSampleCount		= {ParameterName}_TemporalSampleCount;
+	TemporalSampleIndex		= {ParameterName}_TemporalSampleIndex;
+	SequenceFPS				= {ParameterName}_SequenceFPS;
+}
diff --git a/Engine/Plugins/FX/NiagaraMRQ/Source/NiagaraMRQ/NiagaraMRQ.Build.cs b/Engine/Plugins/FX/NiagaraMRQ/Source/NiagaraMRQ/NiagaraMRQ.Build.cs
new file mode 100644
--- /dev/null
+++ b/Engine/Plugins/FX/NiagaraMRQ/Source/NiagaraMRQ/NiagaraMRQ.Build.cs
@@ -0,0 +1,38 @@
+// Copyright Epic Games, Inc. All Rights Reserved.
+
+namespace UnrealBuildTool.Rules
+{
+	public class NiagaraMRQ : ModuleRules
+	{
+		public NiagaraMRQ(ReadOnlyTargetRules Target) : base(Target)
+		{
+			PCHUsage = PCHUsageMode.UseExplicitOrSharedPCHs;
+
+			PublicDependencyModuleNames.AddRange(new[] { "Core" });
+
+			PrivateDependencyModuleNames.AddRange(
+				new[]
+				{
+					"CoreUObject",
+					"Engine",
+					"Slate",
+					"SlateCore",
+					"Projects",
+					"MovieRenderPipelineCore",
+				
+					// Data interface dependencies
+					"Niagara",
+					"NiagaraCore",
+					"VectorVM",
+					"RenderCore",
+					"RHI"
+				}
+			);
+
+			if (Target.bBuildEditor == true)
+			{
+				PrivateDependencyModuleNames.Add("UnrealEd");
+			}
+		}
+	}
+}
diff --git a/Engine/Plugins/FX/NiagaraMRQ/Source/NiagaraMRQ/Private/NiagaraDataInterfaceMRQ.cpp b/Engine/Plugins/FX/NiagaraMRQ/Source/NiagaraMRQ/Private/NiagaraDataInterfaceMRQ.cpp
new file mode 100644
--- /dev/null
+++ b/Engine/Plugins/FX/NiagaraMRQ/Source/NiagaraMRQ/Private/NiagaraDataInterfaceMRQ.cpp
@@ -0,0 +1,206 @@
+// Copyright Epic Games, Inc. All Rights Reserved.
+
+#include "NiagaraDataInterfaceMRQ.h"
+#include "NiagaraCompileHashVisitor.h"
+#include "NiagaraShaderParametersBuilder.h"
+#include "NiagaraSystemInstance.h"
+
+#include "Modules/ModuleManager.h"
+#include "MovieRenderPipelineCoreModule.h"
+
+#define LOCTEXT_NAMESPACE "NiagaraDataInterfaceMRQ"
+
+namespace NDIMRQLocal
+{
+	static const TCHAR* TemplateShaderFile = TEXT("/Plugin/NiagaraMRQ/Private/NiagaraDataInterfaceMRQ.ush");
+
+	const FName NAME_GetMRQInfo(TEXT("GetMRQInfo"));
+
+	struct FNDIInstanceData
+	{
+		bool		bActive = 0;
+		int32		TemporalSampleCount = 1;
+		int32		TemporalSampleIndex = 0;
+		float		SequenceFPS = 1.0f / 30.0f;
+	};
+
+	struct FNDIProxy : public FNiagaraDataInterfaceProxy
+	{
+		virtual int32 PerInstanceDataPassedToRenderThreadSize() const override { return sizeof(FNDIInstanceData); }
+
+		static void ProvidePerInstanceDataForRenderThread(void* InDataToRenderThread, void* InPerInstanceData, const FNiagaraSystemInstanceID& SystemInstance)
+		{
+			FNDIInstanceData* DataToRenderThread = new (InDataToRenderThread) FNDIInstanceData();
+			*DataToRenderThread = *static_cast<const FNDIInstanceData*>(InPerInstanceData);
+		}
+
+		virtual void ConsumePerInstanceDataFromGameThread(void* InDataFromGameThread, const FNiagaraSystemInstanceID& InstanceID) override
+		{
+			const FNDIInstanceData* DataFromGameThread = static_cast<const FNDIInstanceData*>(InDataFromGameThread);
+			FNDIInstanceData& InstanceData_RT = SystemInstancesToInstanceData_RT.FindOrAdd(InstanceID);
+			InstanceData_RT = *DataFromGameThread;
+			DataFromGameThread->~FNDIInstanceData();
+		}
+
+		TMap<FNiagaraSystemInstanceID, FNDIInstanceData> SystemInstancesToInstanceData_RT;
+	};
+
+	void VMGetMRQInfo(FVectorVMExternalFunctionContext& Context)
+	{
+		VectorVM::FUserPtrHandler<FNDIInstanceData> InstanceData(Context);
+		FNDIOutputParam<bool>	OutActive(Context);
+		FNDIOutputParam<int32>	OutTemporalSampleCount(Context);
+		FNDIOutputParam<int32>	OutTemporalSampleIndex(Context);
+		FNDIOutputParam<float>	OutSequenceFPS(Context);
+
+		for (int32 i = 0; i < Context.GetNumInstances(); ++i)
+		{
+			OutActive.SetAndAdvance(InstanceData->bActive);
+			OutTemporalSampleCount.SetAndAdvance(InstanceData->TemporalSampleCount);
+			OutTemporalSampleIndex.SetAndAdvance(InstanceData->TemporalSampleIndex);
+			OutSequenceFPS.SetAndAdvance(InstanceData->SequenceFPS);
+		}
+	}
+}
+
+UNiagaraDataInterfaceMRQ::UNiagaraDataInterfaceMRQ(FObjectInitializer const& ObjectInitializer)
+	: Super(ObjectInitializer)
+{
+	using namespace NDIMRQLocal;
+	Proxy.Reset(new FNDIProxy());
+}
+
+void UNiagaraDataInterfaceMRQ::PostInitProperties()
+{
+	Super::PostInitProperties();
+
+	if (HasAnyFlags(RF_ClassDefaultObject))
+	{
+		ENiagaraTypeRegistryFlags Flags = ENiagaraTypeRegistryFlags::AllowAnyVariable | ENiagaraTypeRegistryFlags::AllowParameter;
+		FNiagaraTypeRegistry::Register(FNiagaraTypeDefinition(GetClass()), Flags);
+	}
+}
+
+void UNiagaraDataInterfaceMRQ::GetFunctions(TArray<FNiagaraFunctionSignature>& OutFunctions)
+{
+	using namespace NDIMRQLocal;
+
+	FNiagaraFunctionSignature& Sig = OutFunctions.AddDefaulted_GetRef();
+	Sig.Name = NAME_GetMRQInfo;
+	Sig.SetDescription(LOCTEXT("GetMRQInfoDesc", "Returns information from MRQ and if it's running or not."));
+	Sig.bMemberFunction = true;
+	Sig.AddInput(FNiagaraVariable(FNiagaraTypeDefinition(GetClass()), TEXT("MRQInterface")));
+	Sig.AddOutput(FNiagaraVariable(FNiagaraTypeDefinition::GetBoolDef(), TEXT("IsActive")), LOCTEXT("IsActiveDesc", "True if we are simulating the frame inside MRQ"));
+	Sig.AddOutput(FNiagaraVariable(FNiagaraTypeDefinition::GetIntDef(), TEXT("TemporalSampleCount")), LOCTEXT("TemporalSampleCountDesc", "Temporal sample count we are rendering with"));
+	Sig.AddOutput(FNiagaraVariable(FNiagaraTypeDefinition::GetIntDef(), TEXT("TemporalSampleIndex")), LOCTEXT("TemporalSampleIndexDesc", "Temporal sample we are simulating this frame for"));
+	Sig.AddOutput(FNiagaraVariable(FNiagaraTypeDefinition::GetFloatDef(), TEXT("SequenceFPS")), LOCTEXT("SequenceFPSDesc", "The source sequence frames per second"));
+}
+
+void UNiagaraDataInterfaceMRQ::GetVMExternalFunction(const FVMExternalFunctionBindingInfo& BindingInfo, void* InstanceData, FVMExternalFunction& OutFunc)
+{
+	using namespace NDIMRQLocal;
+	if (BindingInfo.Name == NAME_GetMRQInfo)
+	{
+		OutFunc = FVMExternalFunction::CreateStatic(VMGetMRQInfo);
+	}
+}
+
+#if WITH_EDITORONLY_DATA
+bool UNiagaraDataInterfaceMRQ::AppendCompileHash(FNiagaraCompileHashVisitor* InVisitor) const
+{
+	using namespace NDIMRQLocal;
+	bool bSuccess = Super::AppendCompileHash(InVisitor);
+	bSuccess &= InVisitor->UpdateShaderFile(TemplateShaderFile);
+	bSuccess &= InVisitor->UpdateShaderParameters<FShaderParameters>();
+	return bSuccess;
+}
+
+bool UNiagaraDataInterfaceMRQ::GetFunctionHLSL(const FNiagaraDataInterfaceGPUParamInfo& ParamInfo, const FNiagaraDataInterfaceGeneratedFunction& FunctionInfo, int FunctionInstanceIndex, FString& OutHLSL)
+{
+	using namespace NDIMRQLocal;
+	return FunctionInfo.DefinitionName == NAME_GetMRQInfo;
+}
+
+void UNiagaraDataInterfaceMRQ::GetParameterDefinitionHLSL(const FNiagaraDataInterfaceGPUParamInfo& ParamInfo, FString& OutHLSL)
+{
+	using namespace NDIMRQLocal;
+	const TMap<FString, FStringFormatArg> TemplateArgs =
+	{
+		{TEXT("ParameterName"),	ParamInfo.DataInterfaceHLSLSymbol},
+	};
+	AppendTemplateHLSL(OutHLSL, TemplateShaderFile, TemplateArgs);
+}
+#endif
+
+void UNiagaraDataInterfaceMRQ::BuildShaderParameters(FNiagaraShaderParametersBuilder& ShaderParametersBuilder) const
+{
+	ShaderParametersBuilder.AddNestedStruct<FShaderParameters>();
+}
+
+void UNiagaraDataInterfaceMRQ::SetShaderParameters(const FNiagaraDataInterfaceSetShaderParametersContext& Context) const
+{
+	using namespace NDIMRQLocal;
+	FNDIProxy& DIProxy = Context.GetProxy<FNDIProxy>();
+	FNDIInstanceData& InstanceData = DIProxy.SystemInstancesToInstanceData_RT.FindChecked(Context.GetSystemInstanceID());
+
+	FShaderParameters* ShaderParameters		= Context.GetParameterNestedStruct<FShaderParameters>();
+	ShaderParameters->Active				= InstanceData.bActive ? 1 : 0;
+	ShaderParameters->TemporalSampleCount	= InstanceData.TemporalSampleCount;
+	ShaderParameters->TemporalSampleIndex	= InstanceData.TemporalSampleIndex;
+	ShaderParameters->SequenceFPS			= InstanceData.SequenceFPS;
+}
+
+bool UNiagaraDataInterfaceMRQ::InitPerInstanceData(void* InPerInstanceData, FNiagaraSystemInstance* SystemInstance)
+{
+	using namespace NDIMRQLocal;
+	FNDIInstanceData* InstanceData = new (InPerInstanceData) FNDIInstanceData();
+	return true;
+}
+
+void UNiagaraDataInterfaceMRQ::DestroyPerInstanceData(void* InPerInstanceData, FNiagaraSystemInstance* SystemInstance)
+{
+	using namespace NDIMRQLocal;
+	FNDIInstanceData* InstanceData = reinterpret_cast<FNDIInstanceData*>(InPerInstanceData);
+	InstanceData->~FNDIInstanceData();
+
+	ENQUEUE_RENDER_COMMAND(RemoveProxy)
+	(
+		[RT_Proxy=GetProxyAs<FNDIProxy>(), InstanceID=SystemInstance->GetId()](FRHICommandListImmediate& CmdList)
+		{
+			RT_Proxy->SystemInstancesToInstanceData_RT.Remove(InstanceID);
+		}
+	);
+}
+
+int32 UNiagaraDataInterfaceMRQ::PerInstanceDataSize() const
+{
+	using namespace NDIMRQLocal;
+	return sizeof(FNDIInstanceData);
+}
+
+bool UNiagaraDataInterfaceMRQ::PerInstanceTick(void* InPerInstanceData, FNiagaraSystemInstance* SystemInstance, float DeltaSeconds)
+{
+	using namespace NDIMRQLocal;
+
+	FNDIInstanceData* InstanceData = reinterpret_cast<FNDIInstanceData*>(InPerInstanceData);
+
+	FMoviePipelineLightweightTickInfo MRQTickInfo;
+	if (FMovieRenderPipelineCoreModule* MRQModule = FModuleManager::Get().GetModulePtr<FMovieRenderPipelineCoreModule>("MovieRenderPipelineCore"))
+	{
+		MRQTickInfo = MRQModule->GetTickInfo();
+	}
+	InstanceData->bActive = MRQTickInfo.bIsActive;
+	InstanceData->TemporalSampleCount = MRQTickInfo.TemporalSampleCount;
+	InstanceData->TemporalSampleIndex = MRQTickInfo.TemporalSampleIndex;
+	InstanceData->SequenceFPS = float(MRQTickInfo.SequenceFPS);
+
+	return false;
+}
+
+void UNiagaraDataInterfaceMRQ::ProvidePerInstanceDataForRenderThread(void* DataForRenderThread, void* PerInstanceData, const FNiagaraSystemInstanceID& SystemInstance)
+{
+	using namespace NDIMRQLocal;
+	FNDIProxy::ProvidePerInstanceDataForRenderThread(DataForRenderThread, PerInstanceData, SystemInstance);
+}
+
+#undef LOCTEXT_NAMESPACE
diff --git a/Engine/Plugins/FX/NiagaraMRQ/Source/NiagaraMRQ/Private/NiagaraDataInterfaceMRQ.h b/Engine/Plugins/FX/NiagaraMRQ/Source/NiagaraMRQ/Private/NiagaraDataInterfaceMRQ.h
new file mode 100644
--- /dev/null
+++ b/Engine/Plugins/FX/NiagaraMRQ/Source/NiagaraMRQ/Private/NiagaraDataInterfaceMRQ.h
@@ -0,0 +1,50 @@
+// Copyright Epic Games, Inc. All Rights Reserved.
+#pragma once
+
+#include "NiagaraDataInterface.h"
+#include "NiagaraDataInterfaceMRQ.generated.h"
+
+enum class ENiagaraSimTarget : uint8;
+struct FNiagaraDataInterfaceGeneratedFunction;
+struct FNiagaraFunctionSignature;
+struct FVMExternalFunctionBindingInfo;
+
+UCLASS(EditInlineNew, Category = "MRQ", meta = (DisplayName = "MovieRenderQueue"))
+class UNiagaraDataInterfaceMRQ : public UNiagaraDataInterface
+{
+	GENERATED_BODY()
+
+	BEGIN_SHADER_PARAMETER_STRUCT(FShaderParameters, )
+		SHADER_PARAMETER(int32,		Active)
+		SHADER_PARAMETER(int32,		TemporalSampleCount)
+		SHADER_PARAMETER(int32,		TemporalSampleIndex)
+		SHADER_PARAMETER(float,		SequenceFPS)
+	END_SHADER_PARAMETER_STRUCT()
+
+public:
+	UNiagaraDataInterfaceMRQ(FObjectInitializer const& ObjectInitializer = FObjectInitializer::Get());
+
+	//UObject Interface
+	virtual void PostInitProperties() override;
+	//UObject Interface End
+
+	//UNiagaraDataInterface Interface
+	virtual void GetFunctions(TArray<FNiagaraFunctionSignature>& OutFunctions) override;
+	virtual void GetVMExternalFunction(const FVMExternalFunctionBindingInfo& BindingInfo, void* InstanceData, FVMExternalFunction &OutFunc) override;
+	virtual bool CanExecuteOnTarget(ENiagaraSimTarget Target) const override { return true; }
+#if WITH_EDITORONLY_DATA
+	virtual bool AppendCompileHash(FNiagaraCompileHashVisitor* InVisitor) const override;
+	virtual bool GetFunctionHLSL(const FNiagaraDataInterfaceGPUParamInfo& ParamInfo, const FNiagaraDataInterfaceGeneratedFunction& FunctionInfo, int FunctionInstanceIndex, FString& OutHLSL) override;
+	virtual void GetParameterDefinitionHLSL(const FNiagaraDataInterfaceGPUParamInfo& ParamInfo, FString& OutHLSL) override;
+#endif
+	virtual void BuildShaderParameters(FNiagaraShaderParametersBuilder& ShaderParametersBuilder) const override;
+	virtual void SetShaderParameters(const FNiagaraDataInterfaceSetShaderParametersContext& Context) const override;
+
+	virtual bool InitPerInstanceData(void* PerInstanceData, FNiagaraSystemInstance* SystemInstance) override;
+	virtual void DestroyPerInstanceData(void* PerInstanceData, FNiagaraSystemInstance* SystemInstance) override;
+	virtual int32 PerInstanceDataSize() const override;
+	virtual bool HasPreSimulateTick() const override { return true; }
+	virtual bool PerInstanceTick(void* PerInstanceData, FNiagaraSystemInstance* SystemInstance, float DeltaSeconds) override;
+	virtual void ProvidePerInstanceDataForRenderThread(void* DataForRenderThread, void* PerInstanceData, const FNiagaraSystemInstanceID& SystemInstance) override;
+	//UNiagaraDataInterface Interface
+};
diff --git a/Engine/Plugins/FX/NiagaraMRQ/Source/NiagaraMRQ/Private/NiagaraMRQModule.cpp b/Engine/Plugins/FX/NiagaraMRQ/Source/NiagaraMRQ/Private/NiagaraMRQModule.cpp
new file mode 100644
--- /dev/null
+++ b/Engine/Plugins/FX/NiagaraMRQ/Source/NiagaraMRQ/Private/NiagaraMRQModule.cpp
@@ -0,0 +1,24 @@
+// Copyright Epic Games, Inc. All Rights Reserved.
+
+#include "NiagaraMRQModule.h"
+#include "Misc/Paths.h"
+#include "Modules/ModuleManager.h"
+#include "ShaderCore.h"
+#include "Interfaces/IPluginManager.h"
+
+#define LOCTEXT_NAMESPACE "FNiagaraMRQModule"
+
+void FNiagaraMRQModule::StartupModule()
+{
+	// map the shader dir so we can use it in the data interface
+	FString PluginShaderDir = FPaths::Combine(IPluginManager::Get().FindPlugin(TEXT("NiagaraMRQ"))->GetBaseDir(), TEXT("Shaders"));
+	AddShaderSourceDirectoryMapping(TEXT("/Plugin/NiagaraMRQ"), PluginShaderDir);
+}
+
+void FNiagaraMRQModule::ShutdownModule()
+{
+}
+
+#undef LOCTEXT_NAMESPACE
+	
+IMPLEMENT_MODULE(FNiagaraMRQModule, NiagaraMRQ)
diff --git a/Engine/Plugins/FX/NiagaraMRQ/Source/NiagaraMRQ/Private/NiagaraMRQModule.h b/Engine/Plugins/FX/NiagaraMRQ/Source/NiagaraMRQ/Private/NiagaraMRQModule.h
new file mode 100644
--- /dev/null
+++ b/Engine/Plugins/FX/NiagaraMRQ/Source/NiagaraMRQ/Private/NiagaraMRQModule.h
@@ -0,0 +1,15 @@
+// Copyright Epic Games, Inc. All Rights Reserved.
+
+#pragma once
+
+#include "CoreMinimal.h"
+#include "Modules/ModuleManager.h"
+
+class FNiagaraMRQModule : public IModuleInterface
+{
+public:
+
+	/** IModuleInterface implementation */
+	virtual void StartupModule() override;
+	virtual void ShutdownModule() override;
+};
diff --git a/Engine/Plugins/MovieScene/MovieRenderPipeline/Source/MovieRenderPipelineCore/Private/MoviePipeline.cpp b/Engine/Plugins/MovieScene/MovieRenderPipeline/Source/MovieRenderPipelineCore/Private/MoviePipeline.cpp
--- a/Engine/Plugins/MovieScene/MovieRenderPipeline/Source/MovieRenderPipelineCore/Private/MoviePipeline.cpp
+++ b/Engine/Plugins/MovieScene/MovieRenderPipeline/Source/MovieRenderPipelineCore/Private/MoviePipeline.cpp
@@ -523,6 +523,9 @@ void UMoviePipeline::TransitionToState(const EMovieRenderPipelineState InNewStat
 			ProcessOutstandingFutures();
 	
 			BeginExport();
+
+			// Clear out MRQ tick information for external consumers as we are done
+			FMovieRenderPipelineCoreModule::SetTickInfo(FMoviePipelineLightweightTickInfo());
 		}
 		break;
 	case EMovieRenderPipelineState::Export:
diff --git a/Engine/Plugins/MovieScene/MovieRenderPipeline/Source/MovieRenderPipelineCore/Private/MoviePipelineTiming.cpp b/Engine/Plugins/MovieScene/MovieRenderPipeline/Source/MovieRenderPipelineCore/Private/MoviePipelineTiming.cpp
--- a/Engine/Plugins/MovieScene/MovieRenderPipeline/Source/MovieRenderPipelineCore/Private/MoviePipelineTiming.cpp
+++ b/Engine/Plugins/MovieScene/MovieRenderPipeline/Source/MovieRenderPipelineCore/Private/MoviePipelineTiming.cpp
@@ -7,6 +7,7 @@
 #include "MoviePipelineOutputBase.h"
 #include "CoreGlobals.h"
 #include "Containers/Array.h"
+#include "Misc/ScopeExit.h"
 #include "MovieRenderPipelineCoreModule.h"
 #include "MovieScene.h"
 #include "LevelSequence.h"
@@ -78,6 +79,16 @@ void UMoviePipeline::ProcessOutstandingFutures()
 
 void UMoviePipeline::TickProducingFrames()
 {
+	ON_SCOPE_EXIT
+	{
+		FMoviePipelineLightweightTickInfo TickInfo;
+		TickInfo.bIsActive = true;
+		TickInfo.TemporalSampleCount = CachedOutputState.TemporalSampleCount;
+		TickInfo.TemporalSampleIndex = CachedOutputState.TemporalSampleIndex;
+		TickInfo.SequenceFPS = GetPipelinePrimaryConfig()->GetEffectiveFrameRate(TargetSequence).AsDecimal();
+		FMovieRenderPipelineCoreModule::SetTickInfo(TickInfo);
+	};
+
 	// The callback for this function does not get registered until Initialization has been called, which sets
 	// the state to Render. If it's not, we have a initialization order/flow issue!
 	check(PipelineState == EMovieRenderPipelineState::ProducingFrames);
diff --git a/Engine/Plugins/MovieScene/MovieRenderPipeline/Source/MovieRenderPipelineCore/Private/MovieRenderPipelineCoreModule.cpp b/Engine/Plugins/MovieScene/MovieRenderPipeline/Source/MovieRenderPipelineCore/Private/MovieRenderPipelineCoreModule.cpp
--- a/Engine/Plugins/MovieScene/MovieRenderPipeline/Source/MovieRenderPipelineCore/Private/MovieRenderPipelineCoreModule.cpp
+++ b/Engine/Plugins/MovieScene/MovieRenderPipeline/Source/MovieRenderPipelineCore/Private/MovieRenderPipelineCoreModule.cpp
@@ -64,6 +64,12 @@ void FMovieRenderPipelineCoreModule::ShutdownModule()
 {
 }
 
+void FMovieRenderPipelineCoreModule::SetTickInfo(const FMoviePipelineLightweightTickInfo& InTickInfo)
+{
+	FMovieRenderPipelineCoreModule& MRQModule = FModuleManager::Get().GetModuleChecked<FMovieRenderPipelineCoreModule>("MovieRenderPipelineCore");
+	MRQModule.TickInfo = InTickInfo;
+}
+
 IMPLEMENT_MODULE(FMovieRenderPipelineCoreModule, MovieRenderPipelineCore);
 DEFINE_LOG_CATEGORY(LogMovieRenderPipeline); 
-DEFINE_LOG_CATEGORY(LogMovieRenderPipelineIO);
\ No newline at end of file
+DEFINE_LOG_CATEGORY(LogMovieRenderPipelineIO);
diff --git a/Engine/Plugins/MovieScene/MovieRenderPipeline/Source/MovieRenderPipelineCore/Public/MovieRenderPipelineCoreModule.h b/Engine/Plugins/MovieScene/MovieRenderPipeline/Source/MovieRenderPipelineCore/Public/MovieRenderPipelineCoreModule.h
--- a/Engine/Plugins/MovieScene/MovieRenderPipeline/Source/MovieRenderPipelineCore/Public/MovieRenderPipelineCoreModule.h
+++ b/Engine/Plugins/MovieScene/MovieRenderPipeline/Source/MovieRenderPipelineCore/Public/MovieRenderPipelineCoreModule.h
@@ -1,3 +1,4 @@
+
 // Copyright Epic Games, Inc. All Rights Reserved.
 
 #pragma once
@@ -12,10 +13,31 @@ class UMoviePipelineExecutorBase;
 class UMoviePipelineQueue;
 class ULevelSequence;
 class UMoviePipeline;
+struct FMoviePipelineFrameOutputState;
 
 // Declare a stat-group for our performance stats to be counted under, readable in game by "stat MovieRenderPipeline".
 DECLARE_STATS_GROUP(TEXT("MovieRenderPipeline"), STATGROUP_MoviePipeline, STATCAT_Advanced);
 
+/**
+* Movie pipeline tick information
+* Used by external consumers, like Niagara, to query information
+*/
+struct FMoviePipelineLightweightTickInfo
+{
+	FMoviePipelineLightweightTickInfo()
+		: bIsActive(false)
+		, TemporalSampleCount(1)
+		, TemporalSampleIndex(0)
+		, SequenceFPS(0.0)
+	{
+	}
+
+	bool	bIsActive;
+	int32	TemporalSampleCount;
+	int32	TemporalSampleIndex;
+	double	SequenceFPS;
+};
+
 namespace MoviePipelineErrorCodes
 {
 	/** Everything completed as expected or we (unfortunately) couldn't detect the error. */
@@ -26,16 +48,17 @@ namespace MoviePipelineErrorCodes
 	constexpr uint8 NoAsset = 2;
 	/** The specified pipeline configuration asset could not be found. Check the logs for details on what it looked for. */
 	constexpr uint8 NoConfig = 3;
-
 }
 
-
 class MOVIERENDERPIPELINECORE_API FMovieRenderPipelineCoreModule : public IModuleInterface
 {
 public:
 	virtual void StartupModule() override;
 	virtual void ShutdownModule() override;
 
+	static void SetTickInfo(const FMoviePipelineLightweightTickInfo& InTickInfo);
+	const FMoviePipelineLightweightTickInfo& GetTickInfo() const { return TickInfo; }
+
 private:
 	bool IsTryingToRenderMovieFromCommandLine(FString& OutSequenceAssetPath, FString& OutConfigAssetPath, FString& OutExecutorType, FString& OutPipelineType) const;
 	void OnMapLoadFinished(class UWorld* InWorld);
@@ -49,6 +72,7 @@ private:
 		UMoviePipelineQueue*& OutQueue, UMoviePipelineExecutorBase*& OutExecutor) const;
 
 private:
+	FMoviePipelineLightweightTickInfo TickInfo;
 	FString MoviePipelineLocalExecutorClassType;
 	FString MoviePipelineClassType;
 	FString SequenceAssetValue;
@@ -64,6 +88,5 @@ public:
 	virtual FText GetEngineChangelistLabel() const { return FText(); }
 };
 
-
 MOVIERENDERPIPELINECORE_API DECLARE_LOG_CATEGORY_EXTERN(LogMovieRenderPipeline, Log, All);
-MOVIERENDERPIPELINECORE_API DECLARE_LOG_CATEGORY_EXTERN(LogMovieRenderPipelineIO, Log, All);
\ No newline at end of file
+MOVIERENDERPIPELINECORE_API DECLARE_LOG_CATEGORY_EXTERN(LogMovieRenderPipelineIO, Log, All);
